Rename/move file(s)

TP_ArchvizBP Should be TP_AEC_ArchvisBP

diff --git a/Templates/TP_AEC_ArchvisBP/TP_AEC_ArchvisBP.uproject b/Templates/TP_AEC_ArchvisBP/TP_AEC_ArchvisBP.uproject
new file mode 100644
--- /dev/null
+++ b/Templates/TP_AEC_ArchvisBP/TP_AEC_ArchvisBP.uproject
@@ -0,0 +1,54 @@
+{
+	"FileVersion": 3,
+	"EngineAssociation": "",
+	"Category": "",
+	"Description": "",
+	"Plugins": [
+		{
+			"Name": "SunPosition",
+			"Enabled": true
+		},
+		{
+			"Name": "DataPrepEditor",
+			"Enabled": true
+		},
+		{
+			"Name": "DatasmithImporter",
+			"Enabled": true
+		},
+		{
+			"Name": "VariantManager",
+			"Enabled": true
+		},
+		{
+			"Name": "PythonScriptPlugin",
+			"Enabled": true
+		},
+		{
+			"Name": "HDRIBackdrop",
+			"Enabled": true
+		},
+		{
+			"Name": "SequencerScripting",
+			"Enabled": false
+		},
+		{
+			"Name": "DatasmithCADImporter",
+			"Enabled": true,
+			"SupportedTargetPlatforms": [
+				"Win64"
+			]
+		},
+		{
+			"Name": "DatasmithIFCImporter",
+			"Enabled": true,
+			"SupportedTargetPlatforms": [
+				"Win64"
+			]
+		},
+		{
+			"Name": "StaticMeshEditorExtension",
+			"Enabled": true
+		}
+	]
+}
\ No newline at end of file
diff --git a/Templates/TP_AEC_ArchvisBP/TP_ArchvisBP.uproject b/Templates/TP_AEC_ArchvisBP/TP_ArchvisBP.uproject
deleted file mode 100644
