Enabled C4D and Mesh Modelling plugins in the AEC_BlankBP and AEC_ArchvisBP Templates #rb none #jira UE-89867

diff --git a/Templates/TP_AEC_ArchvisBP/TP_AEC_ArchvisBP.uproject b/Templates/TP_AEC_ArchvisBP/TP_AEC_ArchvisBP.uproject
--- a/Templates/TP_AEC_ArchvisBP/TP_AEC_ArchvisBP.uproject
+++ b/Templates/TP_AEC_ArchvisBP/TP_AEC_ArchvisBP.uproject
@@ -60,6 +60,18 @@
 			"SupportedTargetPlatforms": [
 				"Win64"
 			]
+		},
+		{
+			"Name": "ModelingToolsEditorMode",
+			"Enabled": true
+		},
+		{
+			"Name": "DatasmithC4DImporter",
+			"Enabled": true,
+			"SupportedTargetPlatforms": [
+				"Win64",
+				"Mac"
+			]
 		}
 	]
 }
\ No newline at end of file
diff --git a/Templates/TP_AEC_BlankBP/TP_AEC_BlankBP.uproject b/Templates/TP_AEC_BlankBP/TP_AEC_BlankBP.uproject
--- a/Templates/TP_AEC_BlankBP/TP_AEC_BlankBP.uproject
+++ b/Templates/TP_AEC_BlankBP/TP_AEC_BlankBP.uproject
@@ -61,6 +61,18 @@
 			"SupportedTargetPlatforms": [
 				"Win64"
 			]
+		},
+		{
+			"Name": "ModelingToolsEditorMode",
+			"Enabled": true
+		},
+		{
+			"Name": "DatasmithC4DImporter",
+			"Enabled": true,
+			"SupportedTargetPlatforms": [
+				"Win64",
+				"Mac"
+			]
 		}
 	]
 }
\ No newline at end of file
