Removed Restricted/NotForLicenceesMoviePipelineDeadlineIntegration plugin from template

diff --git a/Templates/TP_AEC_ArchvisBP/TP_AEC_ArchvisBP.uproject b/Templates/TP_AEC_ArchvisBP/TP_AEC_ArchvisBP.uproject
--- a/Templates/TP_AEC_ArchvisBP/TP_AEC_ArchvisBP.uproject
+++ b/Templates/TP_AEC_ArchvisBP/TP_AEC_ArchvisBP.uproject
@@ -77,10 +77,6 @@
 			"Name": "MoviePipelineMaskRenderPass",
 			"Enabled": true
 		},
-		{
-			"Name": "MoviePipelineDeadlineIntegration",
-			"Enabled": true
-		},
 		{
 			"Name": "GPULightmass",
 			"Enabled": true
diff --git a/Templates/TP_AEC_BlankBP/TP_AEC_BlankBP.uproject b/Templates/TP_AEC_BlankBP/TP_AEC_BlankBP.uproject
--- a/Templates/TP_AEC_BlankBP/TP_AEC_BlankBP.uproject
+++ b/Templates/TP_AEC_BlankBP/TP_AEC_BlankBP.uproject
@@ -78,10 +78,6 @@
 			"Name": "MoviePipelineMaskRenderPass",
 			"Enabled": true
 		},
-		{
-			"Name": "MoviePipelineDeadlineIntegration",
-			"Enabled": true
-		},
 		{
 			"Name": "GPULightmass",
 			"Enabled": true
