Added Electra media decoder plugin to TP_AEC_ArchvisBP and changed sample MediaPlayer decoder to Electra #preflight none

diff --git a/Templates/TP_AEC_ArchvisBP/TP_AEC_ArchvisBP.uproject b/Templates/TP_AEC_ArchvisBP/TP_AEC_ArchvisBP.uproject
--- a/Templates/TP_AEC_ArchvisBP/TP_AEC_ArchvisBP.uproject
+++ b/Templates/TP_AEC_ArchvisBP/TP_AEC_ArchvisBP.uproject
@@ -80,6 +80,10 @@
 		{
 			"Name": "MoviePipelineMaskRenderPass",
 			"Enabled": true
+		},
+		{
+			"Name": "ElectraPlayer",
+			"Enabled": true
 		}
 	]
 }
\ No newline at end of file
