Added Remote Control Panel sample to control Rendering Settings in TP_AEC_ArchvisBP template #preflight none

diff --git a/Templates/TP_AEC_ArchvisBP/TP_AEC_ArchvisBP.uproject b/Templates/TP_AEC_ArchvisBP/TP_AEC_ArchvisBP.uproject
--- a/Templates/TP_AEC_ArchvisBP/TP_AEC_ArchvisBP.uproject
+++ b/Templates/TP_AEC_ArchvisBP/TP_AEC_ArchvisBP.uproject
@@ -84,6 +84,10 @@
 		{
 			"Name": "ElectraPlayer",
 			"Enabled": true
+		},
+		{
+			"Name": "RemoteControl",
+			"Enabled": true
 		}
 	]
 }
\ No newline at end of file
