Resaved TP_AEC_ArchvisBP levels with GPULightmass disabled to avoid build Warning #jira UE-137360 #rb jeanmichel.dignard #preflight none

diff --git a/Templates/TP_AEC_ArchvisBP/TP_AEC_ArchvisBP.uproject b/Templates/TP_AEC_ArchvisBP/TP_AEC_ArchvisBP.uproject
--- a/Templates/TP_AEC_ArchvisBP/TP_AEC_ArchvisBP.uproject
+++ b/Templates/TP_AEC_ArchvisBP/TP_AEC_ArchvisBP.uproject
@@ -57,10 +57,6 @@
 			"Name": "HDRIBackdrop",
 			"Enabled": true
 		},
-		{
-			"Name": "GPULightmass",
-			"Enabled": true
-		},
 		{
 			"Name": "PythonScriptPlugin",
 			"Enabled": true
