Added back GPULightmass plugin to TP_AEC_ArchvisBP template #rb none #preflight none

diff --git a/Templates/TP_AEC_ArchvisBP/TP_AEC_ArchvisBP.uproject b/Templates/TP_AEC_ArchvisBP/TP_AEC_ArchvisBP.uproject
--- a/Templates/TP_AEC_ArchvisBP/TP_AEC_ArchvisBP.uproject
+++ b/Templates/TP_AEC_ArchvisBP/TP_AEC_ArchvisBP.uproject
@@ -84,6 +84,10 @@
 		{
 			"Name": "RemoteControl",
 			"Enabled": true
+		},
+		{
+			"Name": "GPULightmass",
+			"Enabled": true
 		}
 	]
 }
\ No newline at end of file
