Remove references to deleted plugin
#rnx

diff --git a/Templates/TP_AEC_ArchvisBP/TP_AEC_ArchvisBP.uproject b/Templates/TP_AEC_ArchvisBP/TP_AEC_ArchvisBP.uproject
--- a/Templates/TP_AEC_ArchvisBP/TP_AEC_ArchvisBP.uproject
+++ b/Templates/TP_AEC_ArchvisBP/TP_AEC_ArchvisBP.uproject
@@ -19,13 +19,6 @@
 				"Win64"
 			]
 		},
-		{
-			"Name": "DatasmithIFCImporter",
-			"Enabled": true,
-			"SupportedTargetPlatforms": [
-				"Win64"
-			]
-		},
 		{
 			"Name": "DatasmithC4DImporter",
 			"Enabled": true,
diff --git a/Templates/TP_AEC_BlankBP/TP_AEC_BlankBP.uproject b/Templates/TP_AEC_BlankBP/TP_AEC_BlankBP.uproject
--- a/Templates/TP_AEC_BlankBP/TP_AEC_BlankBP.uproject
+++ b/Templates/TP_AEC_BlankBP/TP_AEC_BlankBP.uproject
@@ -19,13 +19,6 @@
 				"Win64"
 			]
 		},
-		{
-			"Name": "DatasmithIFCImporter",
-			"Enabled": true,
-			"SupportedTargetPlatforms": [
-				"Win64"
-			]
-		},
 		{
 			"Name": "DatasmithC4DImporter",
 			"Enabled": true,
diff --git a/Templates/TP_AEC_CollabBP/TP_AEC_CollabBP.uproject b/Templates/TP_AEC_CollabBP/TP_AEC_CollabBP.uproject
--- a/Templates/TP_AEC_CollabBP/TP_AEC_CollabBP.uproject
+++ b/Templates/TP_AEC_CollabBP/TP_AEC_CollabBP.uproject
@@ -39,13 +39,6 @@
 				"Win64"
 			]
 		},
-		{
-			"Name": "DatasmithIFCImporter",
-			"Enabled": true,
-			"SupportedTargetPlatforms": [
-				"Win64"
-			]
-		},
 		{
 			"Name": "DatasmithGLTFImporter",
 			"Enabled": true
diff --git a/Templates/TP_MFG_CollabBP/TP_MFG_CollabBP.uproject b/Templates/TP_MFG_CollabBP/TP_MFG_CollabBP.uproject
--- a/Templates/TP_MFG_CollabBP/TP_MFG_CollabBP.uproject
+++ b/Templates/TP_MFG_CollabBP/TP_MFG_CollabBP.uproject
@@ -39,13 +39,6 @@
 				"Win64"
 			]
 		},
-		{
-			"Name": "DatasmithIFCImporter",
-			"Enabled": true,
-			"SupportedTargetPlatforms": [
-				"Win64"
-			]
-		},
 		{
 			"Name": "DatasmithGLTFImporter",
 			"Enabled": true
diff --git a/Templates/TP_PhotoStudioBP/TP_PhotoStudioBP.uproject b/Templates/TP_PhotoStudioBP/TP_PhotoStudioBP.uproject
--- a/Templates/TP_PhotoStudioBP/TP_PhotoStudioBP.uproject
+++ b/Templates/TP_PhotoStudioBP/TP_PhotoStudioBP.uproject
@@ -40,13 +40,6 @@
 				"Win64"
 			]
 		},
-		{
-			"Name": "DatasmithIFCImporter",
-			"Enabled": true,
-			"SupportedTargetPlatforms": [
-				"Win64"
-			]
-		},
 		{
 			"Name": "DatasmithGLTFImporter",
 			"Enabled": true
