Archvis template disable GPULM plugin to resolve Horde CIS warning 401299

diff --git a/Templates/TP_AEC_ArchvisBP/TP_AEC_ArchvisBP.uproject b/Templates/TP_AEC_ArchvisBP/TP_AEC_ArchvisBP.uproject
--- a/Templates/TP_AEC_ArchvisBP/TP_AEC_ArchvisBP.uproject
+++ b/Templates/TP_AEC_ArchvisBP/TP_AEC_ArchvisBP.uproject
@@ -77,10 +77,6 @@
 		{
 			"Name": "RemoteControl",
 			"Enabled": true
-		},
-		{
-			"Name": "GPULightmass",
-			"Enabled": true
 		}
 	]
 }
\ No newline at end of file
