Enabled Lidar plugin in ArchvisBP template

diff --git a/Templates/TP_AEC_ArchvisBP/TP_AEC_ArchvisBP.uproject b/Templates/TP_AEC_ArchvisBP/TP_AEC_ArchvisBP.uproject
--- a/Templates/TP_AEC_ArchvisBP/TP_AEC_ArchvisBP.uproject
+++ b/Templates/TP_AEC_ArchvisBP/TP_AEC_ArchvisBP.uproject
@@ -77,6 +77,15 @@
 		{
 			"Name": "RemoteControl",
 			"Enabled": true
+		},
+		{
+			"Name": "LidarPointCloud",
+			"Enabled": true,
+			"SupportedTargetPlatforms": [
+				"Win64",
+				"Mac",
+				"Linux"
+			]
 		}
 	]
 }
\ No newline at end of file
